PrereqInstaller: Update old text and images

Fix Signing and add Horde task to build installer (From Ben.Marsh)

#rnx

diff --git a/Engine/Source/Programs/PrereqInstaller/CustomAction/CustomAction.csproj b/Engine/Source/Programs/PrereqInstaller/CustomAction/CustomAction.csproj
--- a/Engine/Source/Programs/PrereqInstaller/CustomAction/CustomAction.csproj
+++ b/Engine/Source/Programs/PrereqInstaller/CustomAction/CustomAction.csproj
@@ -134,13 +134,10 @@
       <CopyToOutputDirectory>Always</CopyToOutputDirectory>
     </Content>
   </ItemGroup>
-  <ItemGroup>
-    <None Include="SignOutput.bat" />
-  </ItemGroup>
   <Import Project="$(MSBuildToolsPath)\Microsoft.CSharp.targets" />
   <Import Project="$(WixCATargetsPath)" />
   <Import Project="..\Sign.targets" />
   <PropertyGroup>
-    <PostBuildEvent>call "$(ProjectDir)SignOutput.bat" $(SignOutput) "$(SignTool)" "$(SigningIdentity)" $(TimeStampServer) "$(TargetDir)$(TargetName).CA.dll"</PostBuildEvent>
+    <PostBuildEvent>$(SignToolCall) "$(TargetDir)$(TargetName).CA.dll""</PostBuildEvent>
   </PropertyGroup>
 </Project>
\ No newline at end of file
diff --git a/Engine/Source/Programs/PrereqInstaller/CustomAction/SignOutput.bat b/Engine/Source/Programs/PrereqInstaller/CustomAction/SignOutput.bat
deleted file mode 100755
