PrereqInstaller: Update VCRedist installer to [REDACTED_PHONE] version from https://docs.microsoft.com/en-us/cpp/windows/latest-supported-vc-redist

#rnx

diff --git a/Engine/Source/Programs/PrereqInstaller/Resources/VCRedist/README.txt b/Engine/Source/Programs/PrereqInstaller/Resources/VCRedist/README.txt
new file mode 100644
--- /dev/null
+++ b/Engine/Source/Programs/PrereqInstaller/Resources/VCRedist/README.txt
@@ -0,0 +1,4 @@
+Please see https://docs.microsoft.com/en-us/cpp/windows/latest-supported-vc-redist for latest installer.
+
+Update version numbers to match in:
+Engine\Source\Programs\Windows\BootstrapPackagedGame\Private\BootstrapPackagedGame.cpp
\ No newline at end of file
