Update vcredist to 14.36.32532

#rnx

diff --git a/Engine/Source/Programs/PrereqInstaller/Resources/VCRedist/README.txt b/Engine/Source/Programs/PrereqInstaller/Resources/VCRedist/README.txt
--- a/Engine/Source/Programs/PrereqInstaller/Resources/VCRedist/README.txt
+++ b/Engine/Source/Programs/PrereqInstaller/Resources/VCRedist/README.txt
@@ -1,4 +1,9 @@
 Please see https://docs.microsoft.com/en-us/cpp/windows/latest-supported-vc-redist for latest installer.
 
 Update version numbers to match in:
-Engine\Source\Programs\Windows\BootstrapPackagedGame\Private\BootstrapPackagedGame.cpp
\ No newline at end of file
+Engine\Source\Programs\Windows\BootstrapPackagedGame\Private\BootstrapPackagedGame.cpp
+
+Copy redist files from:
+C:\Program Files\Microsoft Visual Studio\2022\Professional\VC\Redist\MSVC\<version>\x86\Microsoft.VC143.CRT
+to:
+Engine\Binaries\ThirdParty\AppLocalDependencies\Win64\Microsoft.VC.CRT
\ No newline at end of file
diff --git a/Engine/Source/Programs/Windows/BootstrapPackagedGame/Private/BootstrapPackagedGame.cpp b/Engine/Source/Programs/Windows/BootstrapPackagedGame/Private/BootstrapPackagedGame.cpp
--- a/Engine/Source/Programs/Windows/BootstrapPackagedGame/Private/BootstrapPackagedGame.cpp
+++ b/Engine/Source/Programs/Windows/BootstrapPackagedGame/Private/BootstrapPackagedGame.cpp
@@ -96,8 +96,8 @@ int InstallMissingPrerequisites(const WCHAR* BaseDirectory, const WCHAR* ExecDir
 			// This minimum should match the version installed by
 			// Engine/Source/Programs/PrereqInstaller/Resources/VCRedist/VC_redist.x64.exe
 			const DWORD RequiredMajor = 14;
-			const DWORD RequiredMinor = 34;
-			const DWORD RequiredBld = 31938;
+			const DWORD RequiredMinor = 36;
+			const DWORD RequiredBld = 32532;
 			const DWORD RequiredRbld = 0;
 
 			const DWORD InstalledMajor = RegGetDwordOrZero(Hkey, L"Major");
