Update redist to 14.40.33810.0

#rnx

diff --git a/Engine/Source/Programs/PrereqInstaller/Resources/VCRedist/README.txt b/Engine/Source/Programs/PrereqInstaller/Resources/VCRedist/README.txt
--- a/Engine/Source/Programs/PrereqInstaller/Resources/VCRedist/README.txt
+++ b/Engine/Source/Programs/PrereqInstaller/Resources/VCRedist/README.txt
@@ -4,6 +4,6 @@ Update version numbers to match in:
 Engine\Source\Programs\Windows\BootstrapPackagedGame\Private\BootstrapPackagedGame.cpp
 
 Copy redist files from:
-C:\Program Files\Microsoft Visual Studio\2022\Professional\VC\Redist\MSVC\<version>\x86\Microsoft.VC143.CRT
+C:\Program Files\Microsoft Visual Studio\2022\Professional\VC\Redist\MSVC\<version>\x64\Microsoft.VC143.CRT
 to:
 Engine\Binaries\ThirdParty\AppLocalDependencies\Win64\Microsoft.VC.CRT
\ No newline at end of file
diff --git a/Engine/Source/Programs/Windows/BootstrapPackagedGame/Private/BootstrapPackagedGame.cpp b/Engine/Source/Programs/Windows/BootstrapPackagedGame/Private/BootstrapPackagedGame.cpp
--- a/Engine/Source/Programs/Windows/BootstrapPackagedGame/Private/BootstrapPackagedGame.cpp
+++ b/Engine/Source/Programs/Windows/BootstrapPackagedGame/Private/BootstrapPackagedGame.cpp
@@ -15,7 +15,7 @@ struct VersionInfo
 
 // This minimum should match the version installed by
 // Engine/Source/Programs/PrereqInstaller/Resources/VCRedist/VC_redist.x64.exe
-static const VersionInfo MinRedistVersion = { 14, 38, 33130, 0 };
+static const VersionInfo MinRedistVersion = { 14, 40, 33810, 0 };
 
 bool IsVersionValid(const VersionInfo& Version, const VersionInfo& MinVersion)
 {
